Script plugin: disable engine parsing for now

diff --git a/Engine/Programs/UnrealHeaderTool/Config/DefaultEngine.ini b/Engine/Programs/UnrealHeaderTool/Config/DefaultEngine.ini
--- a/Engine/Programs/UnrealHeaderTool/Config/DefaultEngine.ini
+++ b/Engine/Programs/UnrealHeaderTool/Config/DefaultEngine.ini
@@ -3,4 +3,4 @@
 +ScriptPlugins="ScriptGeneratorPlugin"
 
 +ScriptSupportedModules=/Script/CoreUObject
-+ScriptSupportedModules=/Script/Engine
\ No newline at end of file
+;+ScriptSupportedModules=/Script/Engine
\ No newline at end of file
