Temp fix for compile error: fatal error LNK1189: library limit of 65535 objects exceeded

diff --git a/Engine/Programs/UnrealHeaderTool/Config/DefaultEngine.ini b/Engine/Programs/UnrealHeaderTool/Config/DefaultEngine.ini
--- a/Engine/Programs/UnrealHeaderTool/Config/DefaultEngine.ini
+++ b/Engine/Programs/UnrealHeaderTool/Config/DefaultEngine.ini
@@ -36,4 +36,4 @@ EngineNativePointerMemberBehavior=Disallow
 EngineObjectPtrMemberBehavior=AllowSilently
 NonEngineNativePointerMemberBehavior=AllowSilently
 NonEngineObjectPtrMemberBehavior=AllowSilently
-AutomaticSettersAndGetters=True
\ No newline at end of file
+AutomaticSettersAndGetters=False
\ No newline at end of file
