Updating C# UHT with new units

#rnx

diff --git a/Engine/Programs/UnrealHeaderTool/Config/DefaultEngine.ini b/Engine/Programs/UnrealHeaderTool/Config/DefaultEngine.ini
--- a/Engine/Programs/UnrealHeaderTool/Config/DefaultEngine.ini
+++ b/Engine/Programs/UnrealHeaderTool/Config/DefaultEngine.ini
@@ -132,6 +132,11 @@ gc.MaxObjectsInProgram=262144
 +Units="CandelaPerMeterSquared"
 +Units="cd/m2"
 +Units="CandelaPerMeter2"
++Units="Nanoseconds"
++Units="ns"
++Units="Microseconds"
++Units="us"
++Units="\u00B5s"
 +Units="Milliseconds"
 +Units="ms"
 +Units="Seconds"
