Fixed typo in the centimeters per second units in UHT.

diff --git a/Engine/Programs/UnrealHeaderTool/Config/DefaultEngine.ini b/Engine/Programs/UnrealHeaderTool/Config/DefaultEngine.ini
--- a/Engine/Programs/UnrealHeaderTool/Config/DefaultEngine.ini
+++ b/Engine/Programs/UnrealHeaderTool/Config/DefaultEngine.ini
@@ -57,7 +57,7 @@ gc.MaxObjectsInProgram=262144
 +Units="Radians"
 +Units="rad"
 +Units="CentimetersPerSecond"
-+Units="c/s"
++Units="cm/s"
 +Units="MetersPerSecond"
 +Units="m/s"
 +Units="KilometersPerHour"
