[adding decibels to defaultengine.ini]

Adding db and Decibels to DefaultEngine.ini so it's possible to use UProperty markup Units

#Synced-CL 50377316

diff --git a/Engine/Programs/UnrealHeaderTool/Config/DefaultEngine.ini b/Engine/Programs/UnrealHeaderTool/Config/DefaultEngine.ini
--- a/Engine/Programs/UnrealHeaderTool/Config/DefaultEngine.ini
+++ b/Engine/Programs/UnrealHeaderTool/Config/DefaultEngine.ini
@@ -115,6 +115,8 @@ gc.MaxObjectsInProgram=262144
 +Units="Nm"
 +Units="kgcm"
 +Units="KilogramCentimetersSquaredPerSecondSquared"
++Units="Decibels"
++Units="dB"
 +Units="kgcm\u00B2/s\u00B2"
 +Units="Hertz"
 +Units="Hz"
